- Point VisionOS DDPI to new icons

diff --git a/Engine/Platforms/VisionOS/Config/DataDrivenPlatformInfo.ini b/Engine/Platforms/VisionOS/Config/DataDrivenPlatformInfo.ini
--- a/Engine/Platforms/VisionOS/Config/DataDrivenPlatformInfo.ini
+++ b/Engine/Platforms/VisionOS/Config/DataDrivenPlatformInfo.ini
@@ -4,8 +4,8 @@ IniParent=IOS
 TargetSettingsIniSectionName=/Script/IOSRuntimeSettings.IOSRuntimeSettings
 GlobalIdentifier=9981B55362F64CF9957058C72123A102
 
-NormalIconPath=Launcher/VisionOS/Platform_VisionOS_24x
-LargeIconPath=Launcher/VisionOS/Platform_VisionOS_128x
+NormalIconPath=/Platforms/VisionOS/Content/Editor/Slate/Launcher/Platform_VisionOS_24x
+LargeIconPath=/Platforms/VisionOS/Content/Editor/Slate/Launcher/Platform_VisionOS_128x
 XLargeIconPath=
 AutoSDKPath=
 TutorialPath=SharingAndReleasing/Mobile/iOS
