- Added InstalledBuild xml for VisionOS, which will strip out VisionOS precompiled bits from IBs until we turn it on (some files are still in Platforms/VisionOS, they are benign, but messy)

diff --git a/Engine/Platforms/VisionOS/Build/InstalledEngineBuild.xml b/Engine/Platforms/VisionOS/Build/InstalledEngineBuild.xml
new file mode 100644
--- /dev/null
+++ b/Engine/Platforms/VisionOS/Build/InstalledEngineBuild.xml
@@ -0,0 +1,69 @@
+<?xml version='1.0' ?>
+<BuildGraph xmlns="http://www.epicgames.com/BuildGraph" xmlns:xsi="http://www.w3.org/2001/XMLSchema-instance" xsi:schemaLocation="http://www.epicgames.com/BuildGraph ../../../../Engine/Build/Graph/Schema.xsd" >
+
+	<Option Name="WithVisionOS" Restrict="true|false" DefaultValue="false" Description="Include the VisionOS target platform"/>
+
+	<Property Name="ExtensionPlatforms" Value="$(ExtensionPlatforms);" If="'$(ExtensionPlatforms)' != ''"/>
+	<Property Name="ExtensionPlatforms" Value="$(ExtensionPlatforms)VisionOS"/>
+
+	<Do If="$(WithVisionOS)">
+		<Expand Name="AddExtensionPlatform" PlatformDir="VisionOS"/>
+	</Do>
+	
+	<Extend Name="Platform_FilterOverrides">
+		<Do If="$(WithVisionOS)">
+			<!-- Update CopyEditorEngineOrPlatformFilter, AutomationToolPaths, etc -->
+			<Property Name="AutomationToolPaths">
+				$(AutomationToolPaths)
+				Engine/Platforms/VisionOS/Source/Programs/AutomationTool/...
+			</Property>
+		</Do>
+	</Extend>
+
+	<Extend Name="Platform_AddInstalledRequirements">
+		<Do If="$(WithVisionOS)">
+			<Property Name="InstalledRequirements" Value="$(InstalledRequirements);#UnrealGame VisionOS"/>
+		</Do>
+	</Extend>
+
+	<Extend Name="Platform_Append_DDCPlatformsWin64">
+		<Do If="$(WithVisionOS)">
+			<Property Name="DDCPlatformsMac" Value="$(DDCPlatformsMac)+VisionOS"/>
+		</Do>
+	</Extend>
+
+	<Extend Name="Platform_TagDDCFiles">
+		<Do If="$(WithVisionOS)">
+			<!-- Copy utilities/tools -->
+			<Tag Files="Engine\Binaries\DotNET\VisionOS\..." With="#ToCopy"/>
+		</Do>
+	</Extend>
+
+	<Extend Name="Platform_Tag_UnrealGame">
+		<Do If="$(WithVisionOS)">
+			<Property Name="CopyVisionOSFilter">
+				<!-- Build Files -->
+				Engine/Build/VisionOS/...
+				Engine/Platforms/VisionOS/Binaries/...
+				Engine/Platforms/VisionOS/Build/...
+				Engine/Platforms/VisionOS/Content/Editor/...
+				Engine/Build/VisionOS/...
+			</Property>
+			<Property Name="CopyVisionOSExceptions">
+			</Property>
+			<Tag Files="#UnrealGame VisionOS" With="#Installed Mac"/>
+			<Property Name="Platforms" Value="$(Platforms)VisionOS;"/>
+			<Property Name="CopyInstalledFilter" Value="$(CopyInstalledFilter);$(CopyVisionOSFilter)"/>
+			<Property Name="CopyInstalledExceptions" Value="$(CopyInstalledExceptions);$(CopyVisionOSExceptions)"/>
+		</Do>
+		<Do If="!$(WithVisionOS)">
+			<Property Name="CopyVisionOSCsToolsExceptions">
+				Engine\Saved\CsTools\Engine\Intermediate\ScriptModules\VisionOS.Automation.json
+			</Property>
+			<Property Name="CopyMacCsToolsExceptions" Value="$(CopyMacCsToolsExceptions);$(CopyVisionOSCsToolsExceptions)"/>
+			<Property Name="CopyWin64CsToolsExceptions" Value="$(CopyWin64CsToolsExceptions);$(CopyVisionOSCsToolsExceptions)"/>
+			<Property Name="CopyLinuxCsToolsExceptions" Value="$(CopyLinuxCsToolsExceptions);$(CopyVisionOSCsToolsExceptions)"/>
+		</Do>
+	</Extend>
+
+</BuildGraph>
