Build fix. (Issue 362035)

Upgrading System.Drawing.Common from 4.7.0 => 4.7.2

diff --git a/Engine/Platforms/VisionOS/Source/Programs/AutomationTool/VisionOS.Automation.csproj b/Engine/Platforms/VisionOS/Source/Programs/AutomationTool/VisionOS.Automation.csproj
--- a/Engine/Platforms/VisionOS/Source/Programs/AutomationTool/VisionOS.Automation.csproj
+++ b/Engine/Platforms/VisionOS/Source/Programs/AutomationTool/VisionOS.Automation.csproj
@@ -52,6 +52,6 @@
     <PackageReference Include="Microsoft.CSharp" Version="4.7.0" />
     <PackageReference Include="System.Data.DataSetExtensions" Version="4.5.0" />
     <PackageReference Include="System.Net.Http" Version="4.3.4" />
-    <PackageReference Include="System.Drawing.Common" Version="4.7.0" />
+    <PackageReference Include="System.Drawing.Common" Version="4.7.2" />
   </ItemGroup>
 </Project>
\ No newline at end of file
