Enable SwiftUI by default in VisionOS mode

diff --git a/Engine/Platforms/VisionOS/Config/VisionOSEngine.ini b/Engine/Platforms/VisionOS/Config/VisionOSEngine.ini
--- a/Engine/Platforms/VisionOS/Config/VisionOSEngine.ini
+++ b/Engine/Platforms/VisionOS/Config/VisionOSEngine.ini
@@ -1,3 +1,3 @@
 [SwiftIntegration]
-bUseSwiftUIMain=false
+bUseSwiftUIMain=true
 bCreateBridgingHeader=false
