Add a new Platform VisionOS Project Settings panel to enable/disable VisionOS's full immersion vs window mode

diff --git a/Engine/Platforms/VisionOS/Config/BaseVisionOSEngine.ini b/Engine/Platforms/VisionOS/Config/BaseVisionOSEngine.ini
new file mode 100644
--- /dev/null
+++ b/Engine/Platforms/VisionOS/Config/BaseVisionOSEngine.ini
@@ -0,0 +1,3 @@
+[/Script/VisionOSRuntimeSettings.VisionOSRuntimeSettings]
+bUseSwiftUIMain=true
+bCreateBridgingHeader=false
diff --git a/Engine/Platforms/VisionOS/Config/VisionOSEngine.ini b/Engine/Platforms/VisionOS/Config/VisionOSEngine.ini
deleted file mode 100644
diff --git a/Engine/Platforms/VisionOS/Source/Developer/VisionOSPlatformEditor/Private/VisionOSPlatformEditorModule.cpp b/Engine/Platforms/VisionOS/Source/Developer/VisionOSPlatformEditor/Private/VisionOSPlatformEditorModule.cpp
new file mode 100644
--- /dev/null
+++ b/Engine/Platforms/VisionOS/Source/Developer/VisionOSPlatformEditor/Private/VisionOSPlatformEditorModule.cpp
@@ -0,0 +1,47 @@
+// Copyright Epic Games, Inc. All Rights Reserved.
+
+#include "ISettingsModule.h"
+#include "Modules/ModuleInterface.h"
+#include "Modules/ModuleManager.h"
+#include "VisionOSRuntimeSettings.h"
+
+#define LOCTEXT_NAMESPACE "FVisionOSPlatformEditorModule"
+
+
+/**
+ * Module for VisionOS as a target platform
+ */
+class FVisionOSPlatformEditorModule
+	: public IModuleInterface
+{
+	// IModuleInterface interface
+
+	virtual void StartupModule() override
+	{
+		// register settings
+		ISettingsModule* SettingsModule = FModuleManager::GetModulePtr<ISettingsModule>("Settings");
+		if (SettingsModule != nullptr)
+		{
+			SettingsModule->RegisterSettings("Project", "Platforms", "VisionOS",
+				LOCTEXT("VisionOSSettingsName", "VisionOS"),
+				LOCTEXT("VisionOSSettingsDescription", "Settings for VisionOS projects"),
+				GetMutableDefault<UVisionOSRuntimeSettings>());
+
+		}
+	}
+
+	virtual void ShutdownModule() override
+	{
+		ISettingsModule* SettingsModule = FModuleManager::GetModulePtr<ISettingsModule>("Settings");
+
+		if (SettingsModule != nullptr)
+		{
+			SettingsModule->UnregisterSettings("Project", "Platforms", "VisionOS");
+		}
+	}
+};
+
+
+IMPLEMENT_MODULE(FVisionOSPlatformEditorModule, VisionOSPlatformEditor);
+
+#undef LOCTEXT_NAMESPACE
diff --git a/Engine/Platforms/VisionOS/Source/Developer/VisionOSPlatformEditor/VisionOSPlatformEditor.Build.cs b/Engine/Platforms/VisionOS/Source/Developer/VisionOSPlatformEditor/VisionOSPlatformEditor.Build.cs
new file mode 100644
--- /dev/null
+++ b/Engine/Platforms/VisionOS/Source/Developer/VisionOSPlatformEditor/VisionOSPlatformEditor.Build.cs
@@ -0,0 +1,35 @@
+// Copyright Epic Games, Inc. All Rights Reserved.
+
+using UnrealBuildTool;
+
+public class VisionOSPlatformEditor : ModuleRules
+{
+	public VisionOSPlatformEditor(ReadOnlyTargetRules Target) : base(Target)
+	{
+		BinariesSubFolder = "VisionOS";
+
+		PrivateDependencyModuleNames.AddRange(
+			new string[] {
+				"Core",
+				"CoreUObject",
+				"Engine",
+				"Slate",
+				"PropertyEditor",
+				"SlateCore",
+				"VisionOSRuntimeSettings",
+            }
+		);
+
+		PrivateIncludePathModuleNames.AddRange(
+			new string[] {
+				"Settings",
+			}
+		);
+
+		DynamicallyLoadedModuleNames.AddRange(
+			new string[] {
+				"GameProjectGeneration",
+				}
+		);
+	}
+}
diff --git a/Engine/Platforms/VisionOS/Source/Programs/UnrealBuildTool/UEBuildVisionOS.cs b/Engine/Platforms/VisionOS/Source/Programs/UnrealBuildTool/UEBuildVisionOS.cs
--- a/Engine/Platforms/VisionOS/Source/Programs/UnrealBuildTool/UEBuildVisionOS.cs
+++ b/Engine/Platforms/VisionOS/Source/Programs/UnrealBuildTool/UEBuildVisionOS.cs
@@ -70,6 +70,11 @@ namespace UnrealBuildTool
 				    }
 			    }
 			}
+
+			if (ModuleName == "UnrealEd")
+			{
+				Rules.DynamicallyLoadedModuleNames.Add("VisionOSPlatformEditor");
+			}
 		}
 
 		/// <summary>
diff --git a/Engine/Platforms/VisionOS/Source/Runtime/VisionOSRuntimeSettings/Classes/VisionOSRuntimeSettings.h b/Engine/Platforms/VisionOS/Source/Runtime/VisionOSRuntimeSettings/Classes/VisionOSRuntimeSettings.h
new file mode 100644
--- /dev/null
+++ b/Engine/Platforms/VisionOS/Source/Runtime/VisionOSRuntimeSettings/Classes/VisionOSRuntimeSettings.h
@@ -0,0 +1,41 @@
+// Copyright Epic Games, Inc. All Rights Reserved.
+
+#pragma once
+
+#include "CoreMinimal.h"
+#include "UObject/Object.h"
+#include "Logging/LogMacros.h"
+
+#include "VisionOSRuntimeSettings.generated.h"
+
+DECLARE_LOG_CATEGORY_EXTERN(LogVisionOSRuntimeSettings, Log, All);
+
+/**
+ * Implements the settings for the VisionOS target platform.
+ */
+UCLASS(config=Engine, defaultconfig)
+class VISIONOSRUNTIMESETTINGS_API UVisionOSRuntimeSettings : public UObject
+{
+public:
+	GENERATED_UCLASS_BODY()
+
+	// Use VisionOS .ini files in the editor
+	virtual const TCHAR* GetConfigOverridePlatform() const override 
+	{
+		return TEXT("VisionOS");
+	}
+
+#if WITH_EDITOR
+	virtual void PostEditChangeProperty(struct FPropertyChangedEvent& PropertyChangedEvent) override;
+	virtual void PostInitProperties() override;
+#endif // WITH_EDITOR
+
+    // When building for VisionOS, default to Full Immersion mode.  If disabled, switches to Window mode. Note: when enabled,
+	// the @main struct in UESwift.swift becomes the new main(). Plus, UIApplicationSceneManifest is auto-generation for Info.plist
+    UPROPERTY(config, EditAnywhere, Category = Build, Meta = (DisplayName = "Enable VisionOS's Immersion Mode (Experimental)", ConfigRestartRequired = true))
+    bool bUseSwiftUIMain;
+
+	// If checked, the Swift/ObjC bridging headers will be created.
+	UPROPERTY(config, EditAnywhere, Category = Build, Meta = (DisplayName = "Create Swift/ObjC Bridging Headers", ConfigRestartRequired = true))
+    bool bCreateBridgingHeader;
+};
diff --git a/Engine/Platforms/VisionOS/Source/Runtime/VisionOSRuntimeSettings/Private/VisionOSRuntimeSettings.cpp b/Engine/Platforms/VisionOS/Source/Runtime/VisionOSRuntimeSettings/Private/VisionOSRuntimeSettings.cpp
new file mode 100644
--- /dev/null
+++ b/Engine/Platforms/VisionOS/Source/Runtime/VisionOSRuntimeSettings/Private/VisionOSRuntimeSettings.cpp
@@ -0,0 +1,25 @@
+// Copyright Epic Games, Inc. All Rights Reserved.
+
+#include "VisionOSRuntimeSettings.h"
+#include "Misc/ConfigCacheIni.h"
+#include "Modules/ModuleManager.h"
+
+DEFINE_LOG_CATEGORY(LogVisionOSRuntimeSettings);
+
+UVisionOSRuntimeSettings::UVisionOSRuntimeSettings(const FObjectInitializer& ObjectInitializer)
+	: Super(ObjectInitializer)
+{
+}
+
+#if WITH_EDITOR
+void UVisionOSRuntimeSettings::PostEditChangeProperty(struct FPropertyChangedEvent& PropertyChangedEvent)
+{
+}
+
+void UVisionOSRuntimeSettings::PostInitProperties()
+{
+	Super::PostInitProperties();
+}
+#endif // WITH_EDITOR
+
+IMPLEMENT_MODULE(FDefaultModuleImpl, VisionOSRuntimeSettings);
diff --git a/Engine/Platforms/VisionOS/Source/Runtime/VisionOSRuntimeSettings/VisionOSRuntimeSettings.Build.cs b/Engine/Platforms/VisionOS/Source/Runtime/VisionOSRuntimeSettings/VisionOSRuntimeSettings.Build.cs
new file mode 100644
--- /dev/null
+++ b/Engine/Platforms/VisionOS/Source/Runtime/VisionOSRuntimeSettings/VisionOSRuntimeSettings.Build.cs
@@ -0,0 +1,17 @@
+// Copyright Epic Games, Inc. All Rights Reserved.
+
+using UnrealBuildTool;
+
+public class VisionOSRuntimeSettings : ModuleRules
+{
+	public VisionOSRuntimeSettings(ReadOnlyTargetRules Target) : base(Target)
+	{
+        PrivateDependencyModuleNames.AddRange(
+			new string[]
+			{
+				"Core",
+				"CoreUObject",
+			}
+		);
+	}
+}
diff --git a/Engine/Source/Programs/UnrealBuildTool/Platform/Apple/AppleExports.cs b/Engine/Source/Programs/UnrealBuildTool/Platform/Apple/AppleExports.cs
--- a/Engine/Source/Programs/UnrealBuildTool/Platform/Apple/AppleExports.cs
+++ b/Engine/Source/Programs/UnrealBuildTool/Platform/Apple/AppleExports.cs
@@ -62,8 +62,8 @@ namespace UnrealBuildTool
 		public static void GetSwiftIntegrationSettings(FileReference? ProjectFile, UnrealTargetPlatform Platform, out bool bUseSwiftUIMain, out bool bCreateBridgingHeader)
 		{
 			ConfigHierarchy Ini = ConfigCache.ReadHierarchy(ConfigHierarchyType.Engine, ProjectFile?.Directory, Platform);
-			Ini.TryGetValue("SwiftIntegration", "bUseSwiftUIMain", out bUseSwiftUIMain);
-			Ini.TryGetValue("SwiftIntegration", "bCreateBridgingHeader", out bCreateBridgingHeader);
+			Ini.TryGetValue("/Script/VisionOSRuntimeSettings.VisionOSRuntimeSettings", "bUseSwiftUIMain", out bUseSwiftUIMain);
+			Ini.TryGetValue("/Script/VisionOSRuntimeSettings.VisionOSRuntimeSettings", "bCreateBridgingHeader", out bCreateBridgingHeader);
 		}
 
 		/// <summary>
