Fix or selectively disable issues in UBT's VisionOS platform code. Add support for a minimum visionOS version being supplied to the generated Info.plist file.

#Synced-CL 48713665

diff --git a/Engine/Platforms/VisionOS/Source/Programs/UnrealBuildTool/UEDeployVisionOS.cs b/Engine/Platforms/VisionOS/Source/Programs/UnrealBuildTool/UEDeployVisionOS.cs
--- a/Engine/Platforms/VisionOS/Source/Programs/UnrealBuildTool/UEDeployVisionOS.cs
+++ b/Engine/Platforms/VisionOS/Source/Programs/UnrealBuildTool/UEDeployVisionOS.cs
@@ -23,7 +23,7 @@ namespace UnrealBuildTool
 			return "VisionOS";
 		}
 
-		public static bool GenerateVisionOSPList(string ProjectDirectory, bool bIsUnrealGame, string GameName, bool bIsClient, string ProjectName, string InEngineDir, string AppDirectory, UnrealPluginLanguage? UPL, string? BundleID, ILogger Logger)
+		public static bool GenerateVisionOSPList(string ProjectDirectory, bool bIsUnrealGame, string GameName, bool bIsClient, string ProjectName, string InEngineDir, string AppDirectory, UnrealPluginLanguage? UPL, string? BundleID)
 		{
 			// @todo visionos: THIS!
 
@@ -64,20 +64,12 @@ namespace UnrealBuildTool
 			// required capabilities
 			string RequiredCaps = "\t\t<string>arm64</string>\n";
 
-			// minimum iOS version
-			Ini.GetString("/Script/IOSRuntimeSettings.IOSRuntimeSettings", "MinimumiOSVersion", out string MinVersionSetting);
-			string MinVersion = GetMinimumOSVersion(MinVersionSetting, Logger);
+			// @todo Handle visionOS minimum supported version properly.
+			//Ini.GetString("/Script/IOSRuntimeSettings.IOSRuntimeSettings", "MinimumiOSVersion", out string MinVersionSetting);
+			//string MinVersion = GetMinimumOSVersion(MinVersionSetting, Logger);
 
-			// extra plist data
-			Ini.GetString("/Script/IOSRuntimeSettings.IOSRuntimeSettings", "AdditionalPlistData", out string ExtraData);
-
-			// create the final display name, including converting all entities for XML use
-			string FinalDisplayName = BundleDisplayName.Replace("[PROJECT_NAME]", ProjectName).Replace("_", "");
-			FinalDisplayName = FinalDisplayName.Replace("&", "&amp;");
-			FinalDisplayName = FinalDisplayName.Replace("\"", "&quot;");
-			FinalDisplayName = FinalDisplayName.Replace("\'", "&apos;");
-			FinalDisplayName = FinalDisplayName.Replace("<", "&lt;");
-			FinalDisplayName = FinalDisplayName.Replace(">", "&gt;");
+			// @todo Handle visionOS additional plist data properly.
+			//Ini.GetString("/Script/IOSRuntimeSettings.IOSRuntimeSettings", "AdditionalPlistData", out string ExtraData);
 
 			// generate the plist file
 			StringBuilder Text = new();
@@ -194,7 +186,7 @@ namespace UnrealBuildTool
 
 		public override bool GeneratePList(FileReference? ProjectFile, UnrealTargetConfiguration Config, string ProjectDirectory, bool bIsUnrealGame, string GameName, bool bIsClient, string ProjectName, string InEngineDir, string AppDirectory, List<string> UPLScripts, string? BundleID, bool bBuildAsFramework)
 		{
-			return GenerateVisionOSPList(ProjectDirectory, bIsUnrealGame, GameName, bIsClient, ProjectName, InEngineDir, AppDirectory, null, BundleID, Logger);
+			return GenerateVisionOSPList(ProjectDirectory, bIsUnrealGame, GameName, bIsClient, ProjectName, InEngineDir, AppDirectory, null, BundleID);
 		}
 	}
 }
