UnrealBuildTool: IDE0049 Use framework type

#rnx
#Synced-CL 49007367

diff --git a/Engine/Platforms/VisionOS/Source/Programs/UnrealBuildTool/UEDeployVisionOS.cs b/Engine/Platforms/VisionOS/Source/Programs/UnrealBuildTool/UEDeployVisionOS.cs
--- a/Engine/Platforms/VisionOS/Source/Programs/UnrealBuildTool/UEDeployVisionOS.cs
+++ b/Engine/Platforms/VisionOS/Source/Programs/UnrealBuildTool/UEDeployVisionOS.cs
@@ -42,7 +42,7 @@ namespace UnrealBuildTool
 			// get the settings from the ini file
 			// plist replacements
 			// @todo tvos: Are we going to make TVOS specific .ini files?
-			DirectoryReference? DirRef = bIsUnrealGame ? (!string.IsNullOrEmpty(UnrealBuildTool.GetRemoteIniPath()) ? new DirectoryReference(UnrealBuildTool.GetRemoteIniPath()!) : null) : new DirectoryReference(ProjectDirectory);
+			DirectoryReference? DirRef = bIsUnrealGame ? (!String.IsNullOrEmpty(UnrealBuildTool.GetRemoteIniPath()) ? new DirectoryReference(UnrealBuildTool.GetRemoteIniPath()!) : null) : new DirectoryReference(ProjectDirectory);
 			ConfigHierarchy Ini = ConfigCache.ReadHierarchy(ConfigHierarchyType.Engine, DirRef, UnrealTargetPlatform.IOS);
 
 			// bundle display name
@@ -50,7 +50,7 @@ namespace UnrealBuildTool
 
 			// bundle identifier
 			Ini.GetString("/Script/IOSRuntimeSettings.IOSRuntimeSettings", "BundleIdentifier", out string BundleIdentifier);
-			if (!string.IsNullOrEmpty(BundleID))
+			if (!String.IsNullOrEmpty(BundleID))
 			{
 				BundleIdentifier = BundleID;
 			}
@@ -80,33 +80,33 @@ namespace UnrealBuildTool
 			Text.AppendLine("\t<key>CFBundleDevelopmentRegion</key>");
 			Text.AppendLine("\t<string>en</string>");
 			Text.AppendLine("\t<key>CFBundleDisplayName</key>");
-			Text.AppendLine(string.Format("\t<string>{0}</string>", EncodeBundleName(BundleDisplayName, ProjectName)));
+			Text.AppendLine(String.Format("\t<string>{0}</string>", EncodeBundleName(BundleDisplayName, ProjectName)));
 			Text.AppendLine("\t<key>CFBundleExecutable</key>");
 			string BundleExecutable = bIsUnrealGame ?
 				(bIsClient ? "UnrealClient" : "UnrealGame") :
 				(bIsClient ? GameName + "Client" : GameName);
-			Text.AppendLine(string.Format("\t<string>{0}</string>", BundleExecutable));
+			Text.AppendLine(String.Format("\t<string>{0}</string>", BundleExecutable));
 			Text.AppendLine("\t<key>CFBundleIdentifier</key>");
-			Text.AppendLine(string.Format("\t<string>{0}</string>", BundleIdentifier.Replace("[PROJECT_NAME]", ProjectName).Replace("_", "")));
+			Text.AppendLine(String.Format("\t<string>{0}</string>", BundleIdentifier.Replace("[PROJECT_NAME]", ProjectName).Replace("_", "")));
 			Text.AppendLine("\t<key>CFBundleInfoDictionaryVersion</key>");
 			Text.AppendLine("\t<string>6.0</string>");
 			Text.AppendLine("\t<key>CFBundleName</key>");
-			Text.AppendLine(string.Format("\t<string>{0}</string>", EncodeBundleName(BundleName, ProjectName)));
+			Text.AppendLine(String.Format("\t<string>{0}</string>", EncodeBundleName(BundleName, ProjectName)));
 			Text.AppendLine("\t<key>CFBundlePackageType</key>");
 			Text.AppendLine("\t<string>APPL</string>");
 			Text.AppendLine("\t<key>CFBundleSignature</key>");
 			Text.AppendLine("\t<string>????</string>");
 			Text.AppendLine("\t<key>CFBundleVersion</key>");
-			Text.AppendLine(string.Format("\t<string>{0}</string>", VersionUtilities.UpdateBundleVersion(OldPListData)));
+			Text.AppendLine(String.Format("\t<string>{0}</string>", VersionUtilities.UpdateBundleVersion(OldPListData)));
 			Text.AppendLine("\t<key>CFBundleShortVersionString</key>");
-			Text.AppendLine(string.Format("\t<string>{0}</string>", BundleShortVersion));
+			Text.AppendLine(String.Format("\t<string>{0}</string>", BundleShortVersion));
 			Text.AppendLine("\t<key>LSRequiresIPhoneOS</key>");
 			Text.AppendLine("\t<true/>");
 			Text.AppendLine("\t<key>UIRequiredDeviceCapabilities</key>");
 			Text.AppendLine("\t<array>");
 			foreach (string Line in RequiredCaps.Split("\r\n".ToCharArray()))
 			{
-				if (!string.IsNullOrWhiteSpace(Line))
+				if (!String.IsNullOrWhiteSpace(Line))
 				{
 					Text.AppendLine(Line);
 				}
@@ -127,7 +127,7 @@ namespace UnrealBuildTool
 			Text.AppendLine("\t<string>LaunchScreen</string>");
 
 			// write the iCloud container identifier, if present in the old file
-			if (!string.IsNullOrEmpty(OldPListData))
+			if (!String.IsNullOrEmpty(OldPListData))
 			{
 				int index = OldPListData.IndexOf("ICloudContainerIdentifier");
 				if (index > 0)
@@ -136,7 +136,7 @@ namespace UnrealBuildTool
 					int length = OldPListData.IndexOf("</string>", index) - index;
 					string ICloudContainerIdentifier = OldPListData.Substring(index, length);
 					Text.AppendLine("\t<key>ICloudContainerIdentifier</key>");
-					Text.AppendLine(string.Format("\t<string>{0}</string>", ICloudContainerIdentifier));
+					Text.AppendLine(String.Format("\t<string>{0}</string>", ICloudContainerIdentifier));
 				}
 			}
 
