AutomationScripts: Use default TargetFramework from shared property sheet

#was-backout 50920539
#rnx

#Synced-CL 50923123

diff --git a/Engine/Platforms/VisionOS/Source/Programs/AutomationTool/VisionOS.Automation.csproj b/Engine/Platforms/VisionOS/Source/Programs/AutomationTool/VisionOS.Automation.csproj
--- a/Engine/Platforms/VisionOS/Source/Programs/AutomationTool/VisionOS.Automation.csproj
+++ b/Engine/Platforms/VisionOS/Source/Programs/AutomationTool/VisionOS.Automation.csproj
@@ -2,32 +2,14 @@
   <Import Project="..\..\..\..\..\..\Engine\Source\Programs\Shared\UnrealEngine.csproj.props" />
   
   <PropertyGroup>
-    <TargetFramework>net8.0</TargetFramework>
-    <Configuration Condition=" '$(Configuration)' == '' ">Development</Configuration>
     <OutputType>Library</OutputType>
-    <AppendTargetFrameworkToOutputPath>false</AppendTargetFrameworkToOutputPath>
     <GenerateAssemblyInfo>false</GenerateAssemblyInfo>
     <GenerateTargetFrameworkAttribute>false</GenerateTargetFrameworkAttribute>
-    <Configurations>Debug;Release;Development</Configurations>
     <RootNamespace>VisionOS.Automation</RootNamespace>
     <AssemblyName>VisionOS.Automation</AssemblyName>
     <WarningsNotAsErrors>612,618</WarningsNotAsErrors>
     <OutputPath>..\..\..\Binaries\DotNET\AutomationTool\AutomationScripts\</OutputPath>
     <TreatWarningsAsErrors>true</TreatWarningsAsErrors>
-    <DebugType>pdbonly</DebugType>
-  </PropertyGroup>
-  <PropertyGroup Condition=" '$(Configuration)|$(Platform)' == 'Development|AnyCPU' ">
-    <DefineConstants>$(DefineConstants);TRACE</DefineConstants>
-    <Optimize>true</Optimize>
-  </PropertyGroup>
-  <PropertyGroup Condition=" '$(Configuration)|$(Platform)' == 'Release|AnyCPU' ">
-    <DefineConstants>$(DefineConstants);TRACE</DefineConstants>
-    <Optimize>true</Optimize>
-  </PropertyGroup>
-  <PropertyGroup Condition=" '$(Configuration)|$(Platform)' == 'Debug|AnyCPU' ">
-    <DefineConstants>$(DefineConstants);DEBUG;TRACE</DefineConstants>
-    <DebugSymbols>true</DebugSymbols>
-    <DebugType>full</DebugType>
   </PropertyGroup>
   <ItemGroup>
     <Reference Include="Ionic.Zip.Reduced">
