Remove redundant TargetFramework from csproj

#rnx
#Synced-CL 51506754

diff --git a/Engine/Platforms/VisionOS/Source/Programs/AutomationTool/VisionOS.Automation.csproj b/Engine/Platforms/VisionOS/Source/Programs/AutomationTool/VisionOS.Automation.csproj
--- a/Engine/Platforms/VisionOS/Source/Programs/AutomationTool/VisionOS.Automation.csproj
+++ b/Engine/Platforms/VisionOS/Source/Programs/AutomationTool/VisionOS.Automation.csproj
@@ -24,7 +24,4 @@
     <ProjectReference Include="..\..\..\..\..\Source\Programs\AutomationTool\Scripts\AutomationScripts.Automation.csproj" PrivateAssets="All"><Private>false</Private></ProjectReference>
     <ProjectReference Include="..\..\..\..\..\Source\Programs\AutomationTool\IOS\IOS.Automation.csproj" PrivateAssets="All"><Private>false</Private></ProjectReference>
   </ItemGroup>
-  <ItemGroup>
-    <PackageReference Include="Microsoft.CSharp" Version="4.7.0" />
-  </ItemGroup>
 </Project>
\ No newline at end of file
