Add UE_PLATFORM_* definitions to mirror existing PLATFORM_* definitions for consoles

#Synced-CL 52476816

diff --git a/Engine/Platforms/VisionOS/Source/Programs/UnrealBuildTool/UEBuildVisionOS.cs b/Engine/Platforms/VisionOS/Source/Programs/UnrealBuildTool/UEBuildVisionOS.cs
--- a/Engine/Platforms/VisionOS/Source/Programs/UnrealBuildTool/UEBuildVisionOS.cs
+++ b/Engine/Platforms/VisionOS/Source/Programs/UnrealBuildTool/UEBuildVisionOS.cs
@@ -95,6 +95,7 @@ namespace UnrealBuildTool
 		{
 			base.SetUpEnvironment(Target, CompileEnvironment, LinkEnvironment);
 			CompileEnvironment.Definitions.Add("PLATFORM_VISIONOS=1");
+			CompileEnvironment.Definitions.Add("UE_PLATFORM_VISIONOS=1");
 
 			if (SDKVersionFloat < 2.0)
 			{
